Allow AutoRTFMTests to be staged/packaged

diff --git a/Engine/Programs/AutoRTFMTests/AutoRTFMTests.uproject b/Engine/Programs/AutoRTFMTests/AutoRTFMTests.uproject
new file mode 100644
--- /dev/null
+++ b/Engine/Programs/AutoRTFMTests/AutoRTFMTests.uproject
@@ -0,0 +1,11 @@
+{
+	"FileVersion": 3,
+	"EngineAssociation": "",
+	"Category": "",
+	"Description": "",
+	"TargetPlatforms": [
+		"Windows",
+		"Mac",
+		"Linux"
+	]
+}
diff --git a/Engine/Programs/AutoRTFMTests/Source/AutoRTFMTestsClient.Target.cs b/Engine/Programs/AutoRTFMTests/Source/AutoRTFMTestsClient.Target.cs
new file mode 100644
--- /dev/null
+++ b/Engine/Programs/AutoRTFMTests/Source/AutoRTFMTestsClient.Target.cs
@@ -0,0 +1,7 @@
+// Created by ushell to make blueprint projects work end to end.
+// Placed in Source/ because UAT cleans Intermediate/Source/
+using UnrealBuildTool;
+public class AutoRTFMTestsClientTarget : TargetRules {
+public AutoRTFMTestsClientTarget(TargetInfo Target) : base(Target) {
+Type = TargetType.Client;
+ExtraModuleNames.AddRange(new string[] {"UnrealGame"} ); }}
